This is neccessary to get PGO profiles from the Xbox One- in order to run the non-realtime backend, you can use the command line arg -DeterministicAudio.
rob.gay, phil.popp, aaron.mcleran, maxwell.hayes

diff --git a/Engine/Source/Runtime/AudioMixer/AudioMixer.Build.cs b/Engine/Source/Runtime/AudioMixer/AudioMixer.Build.cs
--- a/Engine/Source/Runtime/AudioMixer/AudioMixer.Build.cs
+++ b/Engine/Source/Runtime/AudioMixer/AudioMixer.Build.cs
@@ -7,8 +7,9 @@ namespace UnrealBuildTool.Rules
 		public AudioMixer(ReadOnlyTargetRules Target) : base(Target)
 		{
 			PrivateIncludePathModuleNames.Add("TargetPlatform");
+            PublicIncludePathModuleNames.Add("TargetPlatform");
 
-			PrivateIncludePaths.AddRange(
+            PrivateIncludePaths.AddRange(
 				new string[]
 				{
 					"Runtime/AudioMixer/Private",
@@ -28,7 +29,8 @@ namespace UnrealBuildTool.Rules
 				{
 					"CoreUObject",
 					"Engine",
-				}
+                    "NonRealtimeAudioRenderer"
+                }
 			);
 
 			AddEngineThirdPartyPrivateStaticDependencies(Target,
diff --git a/Engine/Source/Runtime/AudioMixer/Private/AudioMixerDevice.cpp b/Engine/Source/Runtime/AudioMixer/Private/AudioMixerDevice.cpp
--- a/Engine/Source/Runtime/AudioMixer/Private/AudioMixerDevice.cpp
+++ b/Engine/Source/Runtime/AudioMixer/Private/AudioMixerDevice.cpp
@@ -926,10 +926,14 @@ namespace Audio
 
 	void FMixerDevice::FlushAudioRenderingCommands()
 	{
-		if (IsInitialized() && FPlatformProcess::SupportsMultithreading())
+		if (IsInitialized() && (FPlatformProcess::SupportsMultithreading() && !AudioMixerPlatform->IsNonRealtime()))
 		{
 			SourceManager.FlushCommandQueue();
 		}
+		else if (AudioMixerPlatform->IsNonRealtime())
+		{
+			SourceManager.FlushCommandQueue(true);
+		}
 		else
 		{
 			// Pump the audio device's command queue
diff --git a/Engine/Source/Runtime/AudioMixer/Private/AudioMixerSourceManager.cpp b/Engine/Source/Runtime/AudioMixer/Private/AudioMixerSourceManager.cpp
--- a/Engine/Source/Runtime/AudioMixer/Private/AudioMixerSourceManager.cpp
+++ b/Engine/Source/Runtime/AudioMixer/Private/AudioMixerSourceManager.cpp
@@ -2450,7 +2450,7 @@ namespace Audio
 
 	}
 
-	void FMixerSourceManager::FlushCommandQueue()
+	void FMixerSourceManager::FlushCommandQueue(bool bPumpInCommand)
 	{
 		check(CommandsProcessedEvent != nullptr);
 
@@ -2460,6 +2460,11 @@ namespace Audio
 		// Call update to trigger a final pump of commands
 		Update();
 
+		if (bPumpInCommand)
+		{
+			PumpCommandQueue();
+		}
+
 		// Wait one more time for the double pump
 		CommandsProcessedEvent->Wait();
 	}
diff --git a/Engine/Source/Runtime/AudioMixer/Private/AudioMixerSourceManager.h b/Engine/Source/Runtime/AudioMixer/Private/AudioMixerSourceManager.h
--- a/Engine/Source/Runtime/AudioMixer/Private/AudioMixerSourceManager.h
+++ b/Engine/Source/Runtime/AudioMixer/Private/AudioMixerSourceManager.h
@@ -243,7 +243,7 @@ namespace Audio
 		bool IsBus(const int32 SourceId) const;
 		void PumpCommandQueue();
 		void UpdatePendingReleaseData(bool bForceWait = false);
-		void FlushCommandQueue();
+		void FlushCommandQueue(bool bPumpCommandQueue = false);
 	private:
 
 		void ReleaseSource(const int32 SourceId);
diff --git a/Engine/Source/Runtime/AudioMixer/Public/AudioMixer.h b/Engine/Source/Runtime/AudioMixer/Public/AudioMixer.h
--- a/Engine/Source/Runtime/AudioMixer/Public/AudioMixer.h
+++ b/Engine/Source/Runtime/AudioMixer/Public/AudioMixer.h
@@ -404,6 +404,9 @@ namespace Audio
 		/** Whether or not this platform has hardware decompression. */
 		virtual bool SupportsHardwareDecompression() const { return false; }
 
+		/** Whether this is an interface for a non-realtime renderer. If true, synch events will behave differently to avoid deadlocks. */
+		virtual bool IsNonRealtime() const { return false; }
+
 		/** Creates a Compressed audio info class suitable for decompressing this SoundWave. */
 		virtual ICompressedAudioInfo* CreateCompressedAudioInfo(USoundWave* SoundWave) = 0;
 
diff --git a/Engine/Source/Runtime/Engine/Private/AudioDeviceManager.cpp b/Engine/Source/Runtime/Engine/Private/AudioDeviceManager.cpp
--- a/Engine/Source/Runtime/Engine/Private/AudioDeviceManager.cpp
+++ b/Engine/Source/Runtime/Engine/Private/AudioDeviceManager.cpp
@@ -271,6 +271,8 @@ bool FAudioDeviceManager::LoadDefaultAudioDeviceModule()
 
 	bool bForceNoAudioMixer = FParse::Param(FCommandLine::Get(), TEXT("NoAudioMixer"));
 
+	bool bForceNonRealtimeRenderer = FParse::Param(FCommandLine::Get(), TEXT("DeterministicAudio"));
+
 	// If not using command line switch to use audio mixer, check the game platform engine ini file (e.g. WindowsEngine.ini) which enables it for player
 	bUsingAudioMixer = bForceAudioMixer;
 	if (!bForceAudioMixer && !bForceNoAudioMixer)
@@ -287,6 +289,19 @@ bool FAudioDeviceManager::LoadDefaultAudioDeviceModule()
 	GConfig->GetString(TEXT("Audio"), TEXT("AudioDeviceModuleName"), AudioDeviceModuleName, GEngineIni);
 	GConfig->GetString(TEXT("Audio"), TEXT("AudioMixerModuleName"), AudioMixerModuleName, GEngineIni);
 
+	if (bForceNonRealtimeRenderer)
+	{
+		AudioDeviceModule = FModuleManager::LoadModulePtr<IAudioDeviceModule>(TEXT("NonRealtimeAudioRenderer"));
+
+		static IConsoleVariable* IsUsingAudioMixerCvar = IConsoleManager::Get().FindConsoleVariable(TEXT("au.IsUsingAudioMixer"));
+		check(IsUsingAudioMixerCvar);
+		IsUsingAudioMixerCvar->Set(2, ECVF_SetByConstructor);
+
+		bUsingAudioMixer = true;
+
+		return AudioDeviceModule != nullptr;
+	}
+
 	if (bUsingAudioMixer && AudioMixerModuleName.Len() > 0)
 	{
 		AudioDeviceModule = FModuleManager::LoadModulePtr<IAudioDeviceModule>(*AudioMixerModuleName);
diff --git a/Engine/Source/Runtime/NonRealtimeAudioRenderer/NonRealtimeAudioRenderer.Build.cs b/Engine/Source/Runtime/NonRealtimeAudioRenderer/NonRealtimeAudioRenderer.Build.cs
new file mode 100644
--- /dev/null
+++ b/Engine/Source/Runtime/NonRealtimeAudioRenderer/NonRealtimeAudioRenderer.Build.cs
@@ -0,0 +1,30 @@
+// Copyright 1998-2019 Epic Games, Inc. All Rights Reserved.
+
+namespace UnrealBuildTool.Rules
+{
+	public class NonRealtimeAudioRenderer : ModuleRules
+	{
+		public NonRealtimeAudioRenderer(ReadOnlyTargetRules Target) : base(Target)
+		{
+            PrivateIncludePathModuleNames.Add("TargetPlatform");
+
+            PrivateIncludePaths.AddRange(
+				new string[]
+				{
+					"Runtime/AudioMixer/Private",
+				}
+			);
+
+            PublicIncludePaths.Add("Runtime/AudioMixer/Public");
+
+            PrivateDependencyModuleNames.AddRange(
+            new string[] {
+                    "Core",
+                    "CoreUObject",
+                    "Engine",
+                    "AudioMixer",
+                }
+			);
+        }
+	}
+}
diff --git a/Engine/Source/Runtime/NonRealtimeAudioRenderer/Private/AudioMixerNonRealtime.cpp b/Engine/Source/Runtime/NonRealtimeAudioRenderer/Private/AudioMixerNonRealtime.cpp
new file mode 100644
--- /dev/null
+++ b/Engine/Source/Runtime/NonRealtimeAudioRenderer/Private/AudioMixerNonRealtime.cpp
@@ -0,0 +1,19 @@
+// Copyright 1998-2019 Epic Games, Inc. All Rights Reserved.
+
+#include "AudioMixer.h"
+#include "AudioMixerDevice.h"
+#include "AudioMixerPlatformNonRealtime.h"
+#include "Modules/ModuleManager.h"
+
+class FAudioMixerModuleNonRealtime : public IAudioDeviceModule
+{
+public:
+	virtual bool IsAudioMixerModule() const override { return true; }
+
+	virtual FAudioDevice* CreateAudioDevice() override
+	{
+		return new Audio::FMixerDevice(new Audio::FMixerPlatformNonRealtime());
+	}
+};
+
+IMPLEMENT_MODULE(FAudioMixerModuleNonRealtime, NonRealtimeAudioRenderer);
diff --git a/Engine/Source/Runtime/NonRealtimeAudioRenderer/Private/AudioMixerPlatformNonRealtime.cpp b/Engine/Source/Runtime/NonRealtimeAudioRenderer/Private/AudioMixerPlatformNonRealtime.cpp
new file mode 100644
--- /dev/null
+++ b/Engine/Source/Runtime/NonRealtimeAudioRenderer/Private/AudioMixerPlatformNonRealtime.cpp
@@ -0,0 +1,400 @@
+// Copyright 1998-2019 Epic Games, Inc. All Rights Reserved.
+
+#include "AudioMixerPlatformNonRealtime.h"
+#include "AudioMixer.h"
+#include "AudioMixerDevice.h"
+#include "HAL/PlatformAffinity.h"
+
+#ifndef WITH_XMA2
+#define WITH_XMA2 0
+#endif
+
+#if WITH_XMA2
+#include "XMAAudioInfo.h"
+#endif  //#if WITH_XMA2
+#include "OpusAudioInfo.h"
+#include "VorbisAudioInfo.h"
+#include "Interfaces/IAudioFormat.h"
+
+#include "CoreGlobals.h"
+#include "Misc/ConfigCacheIni.h"
+#include "Misc/MessageDialog.h"
+#include "AudioCompressionSettingsUtils.h"
+
+static int32 DefaultRenderFrameSizeCvar = 256;
+FAutoConsoleVariableRef CVarDefaultRenderFrameSize(
+	TEXT("au.nrt.RenderFrameSize"),
+	DefaultRenderFrameSizeCvar,
+	TEXT("Selects the number of frames to render in a single callback .\n")
+	TEXT("n: Number of frames to render."),
+	ECVF_Default);
+
+static int32 RenderEveryTickCvar = 1;
+FAutoConsoleVariableRef CVarRenderEveryTick(
+	TEXT("au.nrt.RenderEveryTick"),
+	RenderEveryTickCvar,
+	TEXT("When set to 1, calls the RenderAudio call every tick.\n")
+	TEXT("n: Number of frames to render."),
+	ECVF_Default);
+
+namespace Audio
+{
+	FMixerPlatformNonRealtime::FMixerPlatformNonRealtime(float InSampleRate /*= 48000*/, float InNumChannels /*= 2*/, float ExpectedCallbackDuration /*= 0.033f*/)
+		: SampleRate(InSampleRate)
+		, NumChannels(InNumChannels)
+		, TotalDurationRendered(0.0)
+		, TickDelta(0.0)
+		, bIsInitialized(false)
+		, bIsDeviceOpen(false)
+	{
+	}
+
+	FMixerPlatformNonRealtime::~FMixerPlatformNonRealtime()
+	{
+	}
+
+	void FMixerPlatformNonRealtime::RenderAudio(double NumSecondsToRender)
+	{
+		if (!bIsInitialized || !bIsDeviceOpen)
+		{
+			return;
+		}
+
+		const double TimePerCallback = ((double) AudioStreamInfo.NumOutputFrames) / AudioStreamInfo.DeviceInfo.SampleRate;
+		double SecondsRendered = TotalDurationRendered;
+		TotalDurationRendered += NumSecondsToRender;
+
+		while (SecondsRendered < TotalDurationRendered)
+		{
+			// RenderTimeAnalysis.Start();
+			OutputBuffers[CurrentBufferWriteIndex].MixNextBuffer();
+			// RenderTimeAnalysis.End();
+
+			ReadNextBuffer();
+			SecondsRendered += TimePerCallback;
+		}
+	}
+
+	void FMixerPlatformNonRealtime::OpenFileToWriteAudioTo(const FString& OutPath)
+	{
+		// Construct full path:
+		FString AbsoluteFilePath;
+
+		const bool bIsRelativePath = FPaths::IsRelative(OutPath);
+		if (bIsRelativePath)
+		{
+			AbsoluteFilePath = FPaths::ProjectSavedDir() + OutPath;
+			AbsoluteFilePath = FPaths::ConvertRelativePathToFull(AbsoluteFilePath);
+		}
+		else
+		{
+			AbsoluteFilePath = OutPath;
+		}
+
+		FSoundQualityInfo QualityInfo;
+		QualityInfo.SampleRate = SampleRate;
+		QualityInfo.NumChannels = NumChannels;
+		QualityInfo.Quality = 100;
+
+		// Gotcha for bouncing wav files: this has to be filled in.
+		QualityInfo.Duration = 5.0f;
+		QualityInfo.SampleDataSize = 5.0f * SampleRate * NumChannels * sizeof(int16);
+
+		AudioFileWriter.Reset(new FAudioFileWriter(AbsoluteFilePath, QualityInfo));
+	}
+
+	void FMixerPlatformNonRealtime::CloseFile()
+	{
+		AudioFileWriter.Reset();
+	}
+
+	bool FMixerPlatformNonRealtime::InitializeHardware()
+	{
+		if (bIsInitialized)
+		{
+			return false;
+		}
+
+		// Load ogg and vorbis dlls if they haven't been loaded yet
+		LoadVorbisLibraries();
+
+		bIsInitialized = true;
+
+		TickDelta = FApp::GetDeltaTime();
+
+		return true;
+	}
+
+	bool FMixerPlatformNonRealtime::TeardownHardware()
+	{
+		if (!bIsInitialized)
+		{
+			return false;
+		}
+
+		bIsInitialized = false;
+
+		return true;
+	}
+
+	bool FMixerPlatformNonRealtime::IsInitialized() const
+	{
+		return bIsInitialized;
+	}
+
+	bool FMixerPlatformNonRealtime::GetNumOutputDevices(uint32& OutNumOutputDevices)
+	{
+		if (!bIsInitialized)
+		{
+			return false;
+		}
+		
+		OutNumOutputDevices = 1;
+
+		return true;
+	}
+
+	bool FMixerPlatformNonRealtime::GetOutputDeviceInfo(const uint32 InDeviceIndex, FAudioPlatformDeviceInfo& OutInfo)
+	{
+		if (!bIsInitialized)
+		{
+			return false;
+		}
+
+		OutInfo.bIsSystemDefault = true;
+		OutInfo.SampleRate = SampleRate;
+		OutInfo.DeviceId = 0;
+		OutInfo.Format = EAudioMixerStreamDataFormat::Float;
+		OutInfo.Name = TEXT("Non-realtime Renderer");
+		OutInfo.NumChannels = NumChannels;
+
+		OutInfo.OutputChannelArray.Add(EAudioMixerChannel::FrontLeft);
+		OutInfo.OutputChannelArray.Add(EAudioMixerChannel::FrontRight);
+		OutInfo.OutputChannelArray.Add(EAudioMixerChannel::FrontCenter);
+		OutInfo.OutputChannelArray.Add(EAudioMixerChannel::LowFrequency);
+		OutInfo.OutputChannelArray.Add(EAudioMixerChannel::BackLeft);
+		OutInfo.OutputChannelArray.Add(EAudioMixerChannel::BackRight);
+		OutInfo.OutputChannelArray.Add(EAudioMixerChannel::SideLeft);
+		OutInfo.OutputChannelArray.Add(EAudioMixerChannel::SideRight);
+
+		return true;
+	}
+
+	bool FMixerPlatformNonRealtime::GetDefaultOutputDeviceIndex(uint32& OutDefaultDeviceIndex) const
+	{
+		OutDefaultDeviceIndex = 0;
+		return true;
+	}
+
+	bool FMixerPlatformNonRealtime::OpenAudioStream(const FAudioMixerOpenStreamParams& Params)
+	{
+		if (!bIsInitialized)
+		{
+			return false;
+		}
+
+		if (bIsDeviceOpen)
+		{
+			return false;
+		}
+
+		OpenStreamParams = Params;
+		//OpenStreamParams.NumFrames = DefaultRenderFrameSizeCvar;
+		
+
+		AudioStreamInfo.Reset();
+
+		AudioStreamInfo.OutputDeviceIndex = OpenStreamParams.OutputDeviceIndex;
+		AudioStreamInfo.NumOutputFrames = OpenStreamParams.NumFrames;
+		AudioStreamInfo.NumBuffers = OpenStreamParams.NumBuffers;
+		AudioStreamInfo.AudioMixer = OpenStreamParams.AudioMixer;
+
+
+		if (!GetOutputDeviceInfo(AudioStreamInfo.OutputDeviceIndex, AudioStreamInfo.DeviceInfo))
+		{
+			return false;
+		}
+
+		AudioStreamInfo.StreamState = EAudioOutputStreamState::Open;
+		bIsDeviceOpen = true;
+
+		return true;
+	}
+
+	FAudioPlatformDeviceInfo FMixerPlatformNonRealtime::GetPlatformDeviceInfo() const
+	{
+		return AudioStreamInfo.DeviceInfo;
+	}
+
+	bool FMixerPlatformNonRealtime::CloseAudioStream()
+	{
+		if (!bIsInitialized || AudioStreamInfo.StreamState == EAudioOutputStreamState::Closed)
+		{
+			return false;
+		}
+
+		if (bIsDeviceOpen && !StopAudioStream())
+		{
+			return false;
+		}
+
+		bIsDeviceOpen = false;
+
+		AudioStreamInfo.StreamState = EAudioOutputStreamState::Closed;
+
+		return true;
+	}
+
+	bool FMixerPlatformNonRealtime::StartAudioStream()
+	{
+		// Start generating audio with our output source voice
+		BeginGeneratingAudio();
+
+		return true;
+	}
+
+	bool FMixerPlatformNonRealtime::StopAudioStream()
+	{
+		return true;
+	}
+
+	bool FMixerPlatformNonRealtime::CheckAudioDeviceChange()
+	{
+		return false;
+	}
+
+	bool FMixerPlatformNonRealtime::MoveAudioStreamToNewAudioDevice(const FString& InNewDeviceId)
+	{
+		return true;
+	}
+
+	void FMixerPlatformNonRealtime::ResumePlaybackOnNewDevice()
+	{
+		SubmitBuffer(OutputBuffers[CurrentBufferReadIndex].GetBufferData());
+		check(OpenStreamParams.NumFrames * AudioStreamInfo.DeviceInfo.NumChannels == OutputBuffers[CurrentBufferReadIndex].GetBuffer().Num());
+
+		AudioRenderEvent->Trigger();
+	}
+
+	void FMixerPlatformNonRealtime::SubmitBuffer(const uint8* Buffer)
+	{
+		// Do actual buffer submissions here.
+		if (AudioFileWriter.IsValid())
+		{
+			AudioFileWriter->PushAudio((const float*) Buffer, NumChannels * AudioStreamInfo.NumOutputFrames);
+		}
+	}
+
+	FName FMixerPlatformNonRealtime::GetRuntimeFormat(USoundWave* InSoundWave)
+	{
+		//TODO: Set this up to propogate to 
+
+		static FName NAME_OGG(TEXT("OGG"));
+		static FName NAME_OPUS(TEXT("OPUS"));
+		static FName NAME_XMA(TEXT("XMA"));
+
+#if WITH_XMA2 && USE_XMA2_FOR_STREAMING
+		if (InSoundWave->IsStreaming() && InSoundWave->NumChannels <= 2)
+		{
+			return NAME_XMA;
+		}
+#endif
+
+		if (InSoundWave->IsStreaming())
+		{
+#if USE_VORBIS_FOR_STREAMING
+			return NAME_OGG;
+#else
+			return NAME_OPUS;
+#endif
+		}
+
+#if WITH_XMA2
+		if (InSoundWave->NumChannels <= 2)
+		{
+			return NAME_XMA;
+		}
+#endif //#if WITH_XMA2
+
+		return NAME_OGG;
+	}
+
+	bool FMixerPlatformNonRealtime::HasCompressedAudioInfoClass(USoundWave* InSoundWave)
+	{
+		return true;
+	}
+
+	ICompressedAudioInfo* FMixerPlatformNonRealtime::CreateCompressedAudioInfo(USoundWave* InSoundWave)
+	{
+		// TODO: Currently this is a copy paste of the XAudio2 platform interface. Ultimately, this function needs to propogate to the current platform's correct CrateCompressedAudioInfo call.
+
+		check(InSoundWave);
+
+#if WITH_XMA2 && USE_XMA2_FOR_STREAMING
+		if (InSoundWave->IsStreaming() && InSoundWave->NumChannels <= 2 )
+		{
+			return new FXMAAudioInfo();
+		}
+#endif
+
+		if (InSoundWave->IsStreaming())
+		{
+#if USE_VORBIS_FOR_STREAMING
+			return new FVorbisAudioInfo();
+#else
+			return new FOpusAudioInfo();
+#endif
+		}
+
+		static const FName NAME_OGG(TEXT("OGG"));
+		if (FPlatformProperties::RequiresCookedData() ? InSoundWave->HasCompressedData(NAME_OGG) : (InSoundWave->GetCompressedData(NAME_OGG) != nullptr))
+		{
+			return new FVorbisAudioInfo();
+		}
+
+#if WITH_XMA2
+		static const FName NAME_XMA(TEXT("XMA"));
+		if (FPlatformProperties::RequiresCookedData() ? InSoundWave->HasCompressedData(NAME_XMA) : (InSoundWave->GetCompressedData(NAME_XMA) != nullptr))
+		{
+			return new FXMAAudioInfo();
+		}
+#endif
+
+		return nullptr;
+	}
+
+	FString FMixerPlatformNonRealtime::GetDefaultDeviceName()
+	{
+		//GConfig->GetString(TEXT("/Script/WindowsTargetPlatform.WindowsTargetSettings"), TEXT("AudioDevice"), WindowsAudioDeviceName, GEngineIni);
+		return FString();
+	}
+
+	FAudioPlatformSettings FMixerPlatformNonRealtime::GetPlatformSettings() const
+	{
+		return FAudioPlatformSettings::GetPlatformSettings(TEXT("/Script/WindowsTargetPlatform.WindowsTargetSettings"));
+	}
+
+	void FMixerPlatformNonRealtime::OnHardwareUpdate()
+	{
+		if (RenderEveryTickCvar)
+		{
+			RenderAudio(TickDelta);
+		}
+	}
+
+	bool FMixerPlatformNonRealtime::IsNonRealtime() const
+	{
+		return true;
+	}
+
+	uint32 FMixerPlatformNonRealtime::RunInternal()
+	{
+		// Not used.
+		return 0;
+	}
+
+	bool FMixerPlatformNonRealtime::DisablePCMAudioCaching() const
+	{
+		return true;
+	}
+}
\ No newline at end of file
diff --git a/Engine/Source/Runtime/NonRealtimeAudioRenderer/Private/AudioMixerPlatformNonRealtime.h b/Engine/Source/Runtime/NonRealtimeAudioRenderer/Private/AudioMixerPlatformNonRealtime.h
new file mode 100644
--- /dev/null
+++ b/Engine/Source/Runtime/NonRealtimeAudioRenderer/Private/AudioMixerPlatformNonRealtime.h
@@ -0,0 +1,72 @@
+// Copyright 1998-2019 Epic Games, Inc. All Rights Reserved.
+
+#pragma once
+
+#include "AudioMixer.h"
+#include "DSP/FileEncoder.h"
+
+// Any platform defines
+namespace Audio
+{
+	class FMixerPlatformNonRealtime : public IAudioMixerPlatformInterface
+	{
+
+	public:
+
+		FMixerPlatformNonRealtime(float InSampleRate = 48000, float InNumChannels = 2, float ExpectedCallbackDuration = 0.033f);
+		~FMixerPlatformNonRealtime();
+
+		void RenderAudio(double NumSecondsToRender);
+		void OpenFileToWriteAudioTo(const FString& OutPath);
+		void CloseFile();
+
+		//~ Begin IAudioMixerPlatformInterface
+		virtual EAudioMixerPlatformApi::Type GetPlatformApi() const override { return EAudioMixerPlatformApi::XAudio2; }
+		virtual bool InitializeHardware() override;
+		virtual bool CheckAudioDeviceChange() override;
+		virtual bool TeardownHardware() override;
+		virtual bool IsInitialized() const override;
+		virtual bool GetNumOutputDevices(uint32& OutNumOutputDevices) override;
+		virtual bool GetOutputDeviceInfo(const uint32 InDeviceIndex, FAudioPlatformDeviceInfo& OutInfo) override;
+		virtual bool GetDefaultOutputDeviceIndex(uint32& OutDefaultDeviceIndex) const override;
+		virtual bool OpenAudioStream(const FAudioMixerOpenStreamParams& Params) override;
+		virtual bool CloseAudioStream() override;
+		virtual bool StartAudioStream() override;
+		virtual bool StopAudioStream() override;
+		virtual bool MoveAudioStreamToNewAudioDevice(const FString& InNewDeviceId) override;
+		virtual void ResumePlaybackOnNewDevice() override;
+		virtual FAudioPlatformDeviceInfo GetPlatformDeviceInfo() const override;
+		virtual void SubmitBuffer(const uint8* Buffer) override;
+		virtual FName GetRuntimeFormat(USoundWave* InSoundWave) override;
+		virtual bool HasCompressedAudioInfoClass(USoundWave* InSoundWave) override;
+		virtual bool SupportsRealtimeDecompression() const override { return true; }
+		virtual bool DisablePCMAudioCaching() const override;
+		virtual ICompressedAudioInfo* CreateCompressedAudioInfo(USoundWave* InSoundWave) override;
+		virtual FString GetDefaultDeviceName() override;
+		virtual FAudioPlatformSettings GetPlatformSettings() const override;
+		virtual void OnHardwareUpdate() override;
+		virtual bool IsNonRealtime() const override;
+		//~ End IAudioMixerPlatformInterface
+
+	private:
+		float SampleRate;
+		int32 NumChannels;
+
+		// In order to avoid having our number of samples rendered drift, 
+		double TotalDurationRendered;
+
+		// This is retrieved from the tick interval on InitializeHardware.
+		double TickDelta;
+
+		uint32 bIsInitialized : 1;
+		uint32 bIsDeviceOpen : 1;
+
+		TUniquePtr<FAudioFileWriter> AudioFileWriter;
+
+	protected:
+		virtual uint32 RunInternal() override;
+
+	};
+
+}
+
