Fix build error in non-realtime renderer on non PC platforms.

diff --git a/Engine/Source/Runtime/NonRealtimeAudioRenderer/NonRealtimeAudioRenderer.Build.cs b/Engine/Source/Runtime/NonRealtimeAudioRenderer/NonRealtimeAudioRenderer.Build.cs
--- a/Engine/Source/Runtime/NonRealtimeAudioRenderer/NonRealtimeAudioRenderer.Build.cs
+++ b/Engine/Source/Runtime/NonRealtimeAudioRenderer/NonRealtimeAudioRenderer.Build.cs
@@ -25,6 +25,14 @@ namespace UnrealBuildTool.Rules
                     "AudioMixer",
                 }
 			);
+
+            AddEngineThirdPartyPrivateStaticDependencies(Target,
+					"UEOgg",
+					"Vorbis",
+					"VorbisFile",
+					"libOpus",
+					"UELibSampleRate"
+					);
         }
 	}
 }
diff --git a/Engine/Source/Runtime/NonRealtimeAudioRenderer/Private/AudioMixerPlatformNonRealtime.cpp b/Engine/Source/Runtime/NonRealtimeAudioRenderer/Private/AudioMixerPlatformNonRealtime.cpp
--- a/Engine/Source/Runtime/NonRealtimeAudioRenderer/Private/AudioMixerPlatformNonRealtime.cpp
+++ b/Engine/Source/Runtime/NonRealtimeAudioRenderer/Private/AudioMixerPlatformNonRealtime.cpp
@@ -327,7 +327,7 @@ namespace Audio
 	ICompressedAudioInfo* FMixerPlatformNonRealtime::CreateCompressedAudioInfo(USoundWave* InSoundWave)
 	{
 		// TODO: Currently this is a copy paste of the XAudio2 platform interface. Ultimately, this function needs to propogate to the current platform's correct CrateCompressedAudioInfo call.
-
+#if PLATFORM_WINDOWS || WITH_XMA2
 		check(InSoundWave);
 
 #if WITH_XMA2 && USE_XMA2_FOR_STREAMING
@@ -358,7 +358,8 @@ namespace Audio
 		{
 			return new FXMAAudioInfo();
 		}
-#endif
+#endif // WITH_XMA2
+#endif // PLATFORM_WINDOWS || WITH_XMA2
 
 		return nullptr;
 	}
