[NO NON-PII COMMIT MESSAGE]

diff --git a/Engine/Source/Runtime/NonRealtimeAudioRenderer/Private/AudioMixerPlatformNonRealtime.cpp b/Engine/Source/Runtime/NonRealtimeAudioRenderer/Private/AudioMixerPlatformNonRealtime.cpp
--- a/Engine/Source/Runtime/NonRealtimeAudioRenderer/Private/AudioMixerPlatformNonRealtime.cpp
+++ b/Engine/Source/Runtime/NonRealtimeAudioRenderer/Private/AudioMixerPlatformNonRealtime.cpp
@@ -4,6 +4,7 @@
 #include "AudioMixer.h"
 #include "AudioMixerDevice.h"
 #include "HAL/PlatformAffinity.h"
+#include "Misc/App.h"
 
 #ifndef WITH_XMA2
 #define WITH_XMA2 0
