Fix for error logs in PGO. I'm a little worried that these logs mean my original PGO test run with the deterministic audio was not valid, so I'm going to run another PGO test with this to double check.
aaron.mcleran, rob.gay

diff --git a/Engine/Source/Runtime/NonRealtimeAudioRenderer/NonRealtimeAudioRenderer.Build.cs b/Engine/Source/Runtime/NonRealtimeAudioRenderer/NonRealtimeAudioRenderer.Build.cs
--- a/Engine/Source/Runtime/NonRealtimeAudioRenderer/NonRealtimeAudioRenderer.Build.cs
+++ b/Engine/Source/Runtime/NonRealtimeAudioRenderer/NonRealtimeAudioRenderer.Build.cs
@@ -33,6 +33,11 @@ namespace UnrealBuildTool.Rules
 					"libOpus",
 					"UELibSampleRate"
 					);
+
+            if (Target.Platform == UnrealTargetPlatform.XboxOne)
+            {
+                PrivateDependencyModuleNames.Add("XMA2");
+            }
         }
 	}
 }
