Add WITH_XMA2 for NonRealtimeAudioRenderer. Also harden NonRealtimeAudioRenderer.

aaron.mcleran, rob.gay

diff --git a/Engine/Source/Runtime/NonRealtimeAudioRenderer/Private/AudioMixerPlatformNonRealtime.cpp b/Engine/Source/Runtime/NonRealtimeAudioRenderer/Private/AudioMixerPlatformNonRealtime.cpp
--- a/Engine/Source/Runtime/NonRealtimeAudioRenderer/Private/AudioMixerPlatformNonRealtime.cpp
+++ b/Engine/Source/Runtime/NonRealtimeAudioRenderer/Private/AudioMixerPlatformNonRealtime.cpp
@@ -322,13 +322,17 @@ namespace Audio
 
 	bool FMixerPlatformNonRealtime::HasCompressedAudioInfoClass(USoundWave* InSoundWave)
 	{
+#if PLATFORM_WINDOWS || PLATFORM_XBOXONE
 		return true;
+#else
+		return false;
+#endif // PLATFORM_WINDOWS || PLATFORM_XBOXONE
 	}
 
 	ICompressedAudioInfo* FMixerPlatformNonRealtime::CreateCompressedAudioInfo(USoundWave* InSoundWave)
 	{
 		// TODO: Currently this is a copy paste of the XAudio2 platform interface. Ultimately, this function needs to propogate to the current platform's correct CrateCompressedAudioInfo call.
-#if PLATFORM_WINDOWS || WITH_XMA2
+#if PLATFORM_WINDOWS || PLATFORM_XBOXONE
 		check(InSoundWave);
 
 #if WITH_XMA2 && USE_XMA2_FOR_STREAMING
