Fix for XMA2 crash in PGO replay.
aaron.mcleran, dan.oconnor

diff --git a/Engine/Source/Runtime/AudioMixer/Private/AudioMixerDevice.cpp b/Engine/Source/Runtime/AudioMixer/Private/AudioMixerDevice.cpp
--- a/Engine/Source/Runtime/AudioMixer/Private/AudioMixerDevice.cpp
+++ b/Engine/Source/Runtime/AudioMixer/Private/AudioMixerDevice.cpp
@@ -182,6 +182,7 @@ namespace Audio
 			{
 				// Get the platform device info we're using
 				PlatformInfo = AudioMixerPlatform->GetPlatformDeviceInfo();
+				UE_LOG(LogAudioMixer, Display, TEXT("Using Audio Device %s"), *PlatformInfo.Name);
 
 				// Initialize some data that depends on speaker configuration, etc.
 				InitializeChannelAzimuthMap(PlatformInfo.NumChannels);
diff --git a/Engine/Source/Runtime/NonRealtimeAudioRenderer/Private/AudioMixerPlatformNonRealtime.cpp b/Engine/Source/Runtime/NonRealtimeAudioRenderer/Private/AudioMixerPlatformNonRealtime.cpp
--- a/Engine/Source/Runtime/NonRealtimeAudioRenderer/Private/AudioMixerPlatformNonRealtime.cpp
+++ b/Engine/Source/Runtime/NonRealtimeAudioRenderer/Private/AudioMixerPlatformNonRealtime.cpp
@@ -116,6 +116,11 @@ namespace Audio
 			return false;
 		}
 
+#if WITH_XMA2
+		//Initialize our XMA2 decoder context
+		FXMAAudioInfo::Initialize();
+#endif //#if WITH_XMA2
+
 		// Load ogg and vorbis dlls if they haven't been loaded yet
 		LoadVorbisLibraries();
 
@@ -133,6 +138,10 @@ namespace Audio
 			return false;
 		}
 
+#if WITH_XMA2
+		FXMAAudioInfo::Shutdown();
+#endif
+
 		bIsInitialized = false;
 
 		return true;
