Fix for Audio Mixer Non-Realtime Renderer,

diff --git a/Engine/Source/Runtime/NonRealtimeAudioRenderer/Private/AudioMixerPlatformNonRealtime.cpp b/Engine/Source/Runtime/NonRealtimeAudioRenderer/Private/AudioMixerPlatformNonRealtime.cpp
--- a/Engine/Source/Runtime/NonRealtimeAudioRenderer/Private/AudioMixerPlatformNonRealtime.cpp
+++ b/Engine/Source/Runtime/NonRealtimeAudioRenderer/Private/AudioMixerPlatformNonRealtime.cpp
@@ -73,7 +73,10 @@ namespace Audio
 		while (SecondsRendered < TotalDurationRendered)
 		{
 			// RenderTimeAnalysis.Start();
-			OutputBuffers[CurrentBufferWriteIndex].MixNextBuffer();
+			if (OutputBuffers.IsValidIndex(CurrentBufferWriteIndex))
+			{
+				OutputBuffers[CurrentBufferWriteIndex].MixNextBuffer();
+			}
 			// RenderTimeAnalysis.End();
 
 			ReadNextBuffer();
@@ -424,4 +427,4 @@ namespace Audio
 	{
 		return true;
 	}
-}
\ No newline at end of file
+}
