Update codec interfaces to support soundwaveproxy for streaming sounds

diff --git a/Engine/Plugins/Lumin/MagicLeap/Source/MagicLeapAudio/Private/AudioMixerPlatformMagicLeap.cpp b/Engine/Plugins/Lumin/MagicLeap/Source/MagicLeapAudio/Private/AudioMixerPlatformMagicLeap.cpp
--- a/Engine/Plugins/Lumin/MagicLeap/Source/MagicLeapAudio/Private/AudioMixerPlatformMagicLeap.cpp
+++ b/Engine/Plugins/Lumin/MagicLeap/Source/MagicLeapAudio/Private/AudioMixerPlatformMagicLeap.cpp
@@ -443,6 +443,17 @@ namespace Audio
 		return new FADPCMAudioInfo();
 	}
 
+	ICompressedAudioInfo* FMixerPlatformMagicLeap::CreateCompressedAudioInfo(const FSoundWaveProxyPtr& InSoundWave)
+	{
+		if (!(InSoundWave.IsValid() && InSoundWave->IsStreaming()))
+		{
+			return nullptr;
+		}
+
+		static FName NAME_ADPCM(TEXT("ADPCM"));
+		return new FADPCMAudioInfo();
+	}
+
 	FString FMixerPlatformMagicLeap::GetDefaultDeviceName()
 	{
 		return FString(TEXT("MLAudio"));
diff --git a/Engine/Plugins/Lumin/MagicLeap/Source/MagicLeapAudio/Private/AudioMixerPlatformMagicLeap.h b/Engine/Plugins/Lumin/MagicLeap/Source/MagicLeapAudio/Private/AudioMixerPlatformMagicLeap.h
--- a/Engine/Plugins/Lumin/MagicLeap/Source/MagicLeapAudio/Private/AudioMixerPlatformMagicLeap.h
+++ b/Engine/Plugins/Lumin/MagicLeap/Source/MagicLeapAudio/Private/AudioMixerPlatformMagicLeap.h
@@ -38,6 +38,7 @@ namespace Audio
 		virtual FName GetRuntimeFormat(USoundWave* InSoundWave) override;
 		virtual bool HasCompressedAudioInfoClass(USoundWave* InSoundWave) override;
 		virtual ICompressedAudioInfo* CreateCompressedAudioInfo(USoundWave* InSoundWave) override;
+		virtual ICompressedAudioInfo* CreateCompressedAudioInfo(const FSoundWaveProxyPtr& InSoundWave) override;
 		virtual FString GetDefaultDeviceName() override;
 		virtual FAudioPlatformSettings GetPlatformSettings() const override;
 		virtual void SuspendContext() override;
diff --git a/Engine/Plugins/Lumin/MagicLeap/Source/MagicLeapAudio/Private/FakeAudioMixerMagicLeap.cpp b/Engine/Plugins/Lumin/MagicLeap/Source/MagicLeapAudio/Private/FakeAudioMixerMagicLeap.cpp
--- a/Engine/Plugins/Lumin/MagicLeap/Source/MagicLeapAudio/Private/FakeAudioMixerMagicLeap.cpp
+++ b/Engine/Plugins/Lumin/MagicLeap/Source/MagicLeapAudio/Private/FakeAudioMixerMagicLeap.cpp
@@ -184,6 +184,17 @@ namespace Audio
 		return new FADPCMAudioInfo();
 	}
 
+	ICompressedAudioInfo* FFakeAudioMixerMagicLeap::CreateCompressedAudioInfo(const FSoundWaveProxyPtr& InSoundWave)
+	{
+		if (!(InSoundWave.IsValid() && InSoundWave->IsStreaming()))
+		{
+			return nullptr;
+		}
+
+		static FName NAME_ADPCM(TEXT("ADPCM"));
+		return new FADPCMAudioInfo();
+	}
+
 	FString FFakeAudioMixerMagicLeap::GetDefaultDeviceName()
 	{
 		return FString(TEXT("MLAudio"));
diff --git a/Engine/Plugins/Lumin/MagicLeap/Source/MagicLeapAudio/Private/FakeAudioMixerMagicLeap.h b/Engine/Plugins/Lumin/MagicLeap/Source/MagicLeapAudio/Private/FakeAudioMixerMagicLeap.h
--- a/Engine/Plugins/Lumin/MagicLeap/Source/MagicLeapAudio/Private/FakeAudioMixerMagicLeap.h
+++ b/Engine/Plugins/Lumin/MagicLeap/Source/MagicLeapAudio/Private/FakeAudioMixerMagicLeap.h
@@ -38,6 +38,7 @@ namespace Audio
 		virtual int32 GetNumFrames(const int32 InNumReqestedFrames) override;
 		virtual bool HasCompressedAudioInfoClass(USoundWave* InSoundWave) override;
 		virtual ICompressedAudioInfo* CreateCompressedAudioInfo(USoundWave* InSoundWave) override;
+		virtual ICompressedAudioInfo* CreateCompressedAudioInfo(const FSoundWaveProxyPtr& InSoundWave) override;
 		virtual FString GetDefaultDeviceName() override;
 		virtual FAudioPlatformSettings GetPlatformSettings() const override;
 
diff --git a/Engine/Plugins/Runtime/Metasound/Source/MetasoundEngine/Private/MetasoundWavePlayerNode.cpp b/Engine/Plugins/Runtime/Metasound/Source/MetasoundEngine/Private/MetasoundWavePlayerNode.cpp
--- a/Engine/Plugins/Runtime/Metasound/Source/MetasoundEngine/Private/MetasoundWavePlayerNode.cpp
+++ b/Engine/Plugins/Runtime/Metasound/Source/MetasoundEngine/Private/MetasoundWavePlayerNode.cpp
@@ -179,7 +179,7 @@ namespace Metasound
 
 				int32 NumFramesDecoded = Decoder.GenerateAudio(PostSrcBufferPtr, NumOutputFrames, *PitchShiftCents) / NumInputChannels;
 
-				// TODO: handle decoder having completed during it's decode
+				// handle decoder having completed during it's decode
 				if (!Decoder.CanGenerateAudio() ||  (NumFramesDecoded < NumOutputFrames))
 				{
 					bIsPlaying = false;
diff --git a/Engine/Source/Runtime/Android/AudioMixerAndroid/Private/AudioMixerPlatformAndroid.cpp b/Engine/Source/Runtime/Android/AudioMixerAndroid/Private/AudioMixerPlatformAndroid.cpp
--- a/Engine/Source/Runtime/Android/AudioMixerAndroid/Private/AudioMixerPlatformAndroid.cpp
+++ b/Engine/Source/Runtime/Android/AudioMixerAndroid/Private/AudioMixerPlatformAndroid.cpp
@@ -468,6 +468,28 @@ namespace Audio
 #endif // WITH_ENGINE
 	}
 
+	ICompressedAudioInfo* FMixerPlatformAndroid::CreateCompressedAudioInfo(const FSoundWaveProxyPtr& InSoundWave)
+	{
+		if (!ensure(InSoundWave.IsValid() && InSoundWave->IsStreaming()))
+		{
+			return nullptr;
+		}
+
+#if WITH_ENGINE
+		static FName NAME_OGG(TEXT("OGG"));
+		static FName NAME_ADPCM(TEXT("ADPCM"));
+
+		if (InSoundWave->IsSeekableStreaming())
+		{
+			return new FADPCMAudioInfo();
+		}
+
+		return new FVorbisAudioInfo();
+#endif // WITH_ENGINE
+
+		return nullptr;
+	}
+
 	FString FMixerPlatformAndroid::GetDefaultDeviceName()
 	{
 		return FString();
diff --git a/Engine/Source/Runtime/Android/AudioMixerAndroid/Private/AudioMixerPlatformAndroid.h b/Engine/Source/Runtime/Android/AudioMixerAndroid/Private/AudioMixerPlatformAndroid.h
--- a/Engine/Source/Runtime/Android/AudioMixerAndroid/Private/AudioMixerPlatformAndroid.h
+++ b/Engine/Source/Runtime/Android/AudioMixerAndroid/Private/AudioMixerPlatformAndroid.h
@@ -36,6 +36,7 @@ namespace Audio
 		virtual FName GetRuntimeFormat(USoundWave* InSoundWave) override;
 		virtual bool HasCompressedAudioInfoClass(USoundWave* InSoundWave) override;
 		virtual ICompressedAudioInfo* CreateCompressedAudioInfo(USoundWave* InSoundWave) override;
+		virtual ICompressedAudioInfo* CreateCompressedAudioInfo(const FSoundWaveProxyPtr& InSoundWave) override;
 		virtual FString GetDefaultDeviceName() override;
 		virtual FAudioPlatformSettings GetPlatformSettings() const override;
 		virtual void SuspendContext() override;
diff --git a/Engine/Source/Runtime/Apple/AudioMixerAudioUnit/Private/AudioMixerPlatformAudioUnit.cpp b/Engine/Source/Runtime/Apple/AudioMixerAudioUnit/Private/AudioMixerPlatformAudioUnit.cpp
--- a/Engine/Source/Runtime/Apple/AudioMixerAudioUnit/Private/AudioMixerPlatformAudioUnit.cpp
+++ b/Engine/Source/Runtime/Apple/AudioMixerAudioUnit/Private/AudioMixerPlatformAudioUnit.cpp
@@ -403,6 +403,11 @@ namespace Audio
 	{
 		return new FADPCMAudioInfo();
 	}
+
+	ICompressedAudioInfo* FMixerPlatformAudioUnit::CreateCompressedAudioInfo(const FSoundWaveProxyPtr& InSoundWave)
+	{
+		return new FADPCMAudioInfo();
+	}
 	
 	FString FMixerPlatformAudioUnit::GetDefaultDeviceName()
 	{
diff --git a/Engine/Source/Runtime/Apple/AudioMixerAudioUnit/Private/AudioMixerPlatformAudioUnit.h b/Engine/Source/Runtime/Apple/AudioMixerAudioUnit/Private/AudioMixerPlatformAudioUnit.h
--- a/Engine/Source/Runtime/Apple/AudioMixerAudioUnit/Private/AudioMixerPlatformAudioUnit.h
+++ b/Engine/Source/Runtime/Apple/AudioMixerAudioUnit/Private/AudioMixerPlatformAudioUnit.h
@@ -38,6 +38,7 @@ namespace Audio
 		virtual FName GetRuntimeFormat(USoundWave* InSoundWave) override;
 		virtual bool HasCompressedAudioInfoClass(USoundWave* InSoundWave) override;
 		virtual ICompressedAudioInfo* CreateCompressedAudioInfo(USoundWave* InSoundWave) override;
+		virtual ICompressedAudioInfo* CreateCompressedAudioInfo(const FSoundWaveProxyPtr& InSoundWave) override;
 		virtual FString GetDefaultDeviceName() override;
 		virtual FAudioPlatformSettings GetPlatformSettings() const override;
 		virtual int32 GetNumFrames(const int32 InNumReqestedFrames) override;
diff --git a/Engine/Source/Runtime/Apple/AudioMixerCoreAudio/Private/AudioMixerPlatformCoreAudio.cpp b/Engine/Source/Runtime/Apple/AudioMixerCoreAudio/Private/AudioMixerPlatformCoreAudio.cpp
--- a/Engine/Source/Runtime/Apple/AudioMixerCoreAudio/Private/AudioMixerPlatformCoreAudio.cpp
+++ b/Engine/Source/Runtime/Apple/AudioMixerCoreAudio/Private/AudioMixerPlatformCoreAudio.cpp
@@ -414,6 +414,27 @@ namespace Audio
 		return nullptr;
 	}
 
+	ICompressedAudioInfo* FMixerPlatformCoreAudio::CreateCompressedAudioInfo(const FSoundWaveProxyPtr& InSoundWave)
+	{
+#if WITH_ENGINE
+		if (!ensure(InSoundWave.IsValid))
+		{
+			return nullptr;
+		}
+
+
+		if (InSoundWave->IsSeekableStreaming())
+		{
+			return new FADPCMAudioInfo();
+		}
+		else if (InSoundWave->IsStreaming())
+		{
+			return new FOpusAudioInfo();
+		}
+
+		return nullptr;
+	}
+
 	FString FMixerPlatformCoreAudio::GetDefaultDeviceName()
 	{
 		return FString();
diff --git a/Engine/Source/Runtime/Apple/AudioMixerCoreAudio/Private/AudioMixerPlatformCoreAudio.h b/Engine/Source/Runtime/Apple/AudioMixerCoreAudio/Private/AudioMixerPlatformCoreAudio.h
--- a/Engine/Source/Runtime/Apple/AudioMixerCoreAudio/Private/AudioMixerPlatformCoreAudio.h
+++ b/Engine/Source/Runtime/Apple/AudioMixerCoreAudio/Private/AudioMixerPlatformCoreAudio.h
@@ -38,6 +38,7 @@ namespace Audio
 		virtual FName GetRuntimeFormat(USoundWave* InSoundWave) override;
 		virtual bool HasCompressedAudioInfoClass(USoundWave* InSoundWave) override;
 		virtual ICompressedAudioInfo* CreateCompressedAudioInfo(USoundWave* InSoundWave) override;
+		virtual ICompressedAudioInfo* CreateCompressedAudioInfo(const FSoundWaveProxyPtr& InSoundWave) override;
 		virtual FString GetDefaultDeviceName() override;
 		virtual FAudioPlatformSettings GetPlatformSettings() const override;
         virtual int32 GetNumFrames(const int32 InNumReqestedFrames) override;
diff --git a/Engine/Source/Runtime/NonRealtimeAudioRenderer/Private/AudioMixerPlatformNonRealtime.cpp b/Engine/Source/Runtime/NonRealtimeAudioRenderer/Private/AudioMixerPlatformNonRealtime.cpp
--- a/Engine/Source/Runtime/NonRealtimeAudioRenderer/Private/AudioMixerPlatformNonRealtime.cpp
+++ b/Engine/Source/Runtime/NonRealtimeAudioRenderer/Private/AudioMixerPlatformNonRealtime.cpp
@@ -381,6 +381,37 @@ namespace Audio
 			return XMA2_INFO_NEW();
 		}
 #endif // WITH_XMA2
+#endif // PLATFORM_WINDOWS || HAS_COMPRESSED_AUDIO_INFO_CLASS
+
+		return nullptr;
+	}
+
+	ICompressedAudioInfo* FMixerPlatformNonRealtime::CreateCompressedAudioInfo(const FSoundWaveProxyPtr& InSoundWave)
+	{
+		if (!(InSoundWave.IsValid() && InSoundWave->IsStreaming()))
+		{
+			return nullptr;
+		}
+
+		// TODO: Currently this is a copy paste of the XAudio2 platform interface. Ultimately, this function needs to propagate to the current platform's correct CrateCompressedAudioInfo call.
+#if PLATFORM_WINDOWS || HAS_COMPRESSED_AUDIO_INFO_CLASS
+
+#if WITH_XMA2 && USE_XMA2_FOR_STREAMING
+		if (InSoundWave->IsStreaming() && InSoundWave->GetNumChannels() <= 2)
+		{
+			return XMA2_INFO_NEW();
+		}
+#endif
+
+		if (InSoundWave->IsStreaming())
+		{
+#if USE_VORBIS_FOR_STREAMING
+			return new FVorbisAudioInfo();
+#else
+			return new FOpusAudioInfo();
+#endif
+		}
+
 #endif // PLATFORM_WINDOWS || HAS_COMPRESSED_AUDIO_INFO_CLASS
 
 		return nullptr;
diff --git a/Engine/Source/Runtime/NonRealtimeAudioRenderer/Public/AudioMixerPlatformNonRealtime.h b/Engine/Source/Runtime/NonRealtimeAudioRenderer/Public/AudioMixerPlatformNonRealtime.h
--- a/Engine/Source/Runtime/NonRealtimeAudioRenderer/Public/AudioMixerPlatformNonRealtime.h
+++ b/Engine/Source/Runtime/NonRealtimeAudioRenderer/Public/AudioMixerPlatformNonRealtime.h
@@ -42,6 +42,7 @@ namespace Audio
 		virtual bool SupportsRealtimeDecompression() const override { return true; }
 		virtual bool DisablePCMAudioCaching() const override;
 		virtual ICompressedAudioInfo* CreateCompressedAudioInfo(USoundWave* InSoundWave) override;
+		virtual ICompressedAudioInfo* CreateCompressedAudioInfo(const FSoundWaveProxyPtr& InSoundWave) override;
 		virtual FString GetDefaultDeviceName() override;
 		virtual FAudioPlatformSettings GetPlatformSettings() const override;
 		virtual void OnHardwareUpdate() override;
