Audio: Fix for seekable streaming (ADPCM) not being an option for FMixerPlatformNonRealtime, resulting in a data/codec mismatch for offline audio rendering for assets set to seekable + streaming.

diff --git a/Engine/Source/Runtime/NonRealtimeAudioRenderer/Private/AudioMixerPlatformNonRealtime.cpp b/Engine/Source/Runtime/NonRealtimeAudioRenderer/Private/AudioMixerPlatformNonRealtime.cpp
--- a/Engine/Source/Runtime/NonRealtimeAudioRenderer/Private/AudioMixerPlatformNonRealtime.cpp
+++ b/Engine/Source/Runtime/NonRealtimeAudioRenderer/Private/AudioMixerPlatformNonRealtime.cpp
@@ -20,6 +20,7 @@
 #endif  //#if WITH_XMA2
 #include "OpusAudioInfo.h"
 #include "VorbisAudioInfo.h"
+#include "ADPCMAudioInfo.h"
 #include "Interfaces/IAudioFormat.h"
 
 #include "CoreGlobals.h"
@@ -310,6 +311,12 @@ namespace Audio
 		static FName NAME_OGG(TEXT("OGG"));
 		static FName NAME_OPUS(TEXT("OPUS"));
 		static FName NAME_XMA(TEXT("XMA"));
+		static FName NAME_ADPCM(TEXT("ADPCM"));
+
+		if (InSoundWave->IsSeekableStreaming())
+		{
+			return NAME_ADPCM;
+		}
 
 #if WITH_XMA2 && USE_XMA2_FOR_STREAMING
 		if (InSoundWave->IsStreaming(nullptr) && InSoundWave->NumChannels <= 2)
@@ -355,12 +362,21 @@ namespace Audio
 #if WITH_XMA2 && USE_XMA2_FOR_STREAMING
 		if (InSoundWave->IsStreaming() && InSoundWave->NumChannels <= 2 )
 		{
+			if (InSoundWave->IsSeekableStreaming())
+			{
+				return new FADPCMAudioInfo();
+			}
 			return XMA2_INFO_NEW();
 		}
 #endif
 
 		if (InSoundWave->IsStreaming(nullptr))
 		{
+			if (InSoundWave->IsSeekableStreaming())
+			{
+				return new FADPCMAudioInfo();
+			}
+
 #if USE_VORBIS_FOR_STREAMING
 			return new FVorbisAudioInfo();
 #else
@@ -405,6 +421,11 @@ namespace Audio
 
 		if (InSoundWave->IsStreaming())
 		{
+			if (InSoundWave->IsSeekableStreaming())
+			{
+				return new FADPCMAudioInfo();
+			}
+
 #if USE_VORBIS_FOR_STREAMING
 			return new FVorbisAudioInfo();
 #else
