Fix for non-realtime renderer

diff --git a/Engine/Source/Runtime/NonRealtimeAudioRenderer/Private/AudioMixerPlatformNonRealtime.cpp b/Engine/Source/Runtime/NonRealtimeAudioRenderer/Private/AudioMixerPlatformNonRealtime.cpp
--- a/Engine/Source/Runtime/NonRealtimeAudioRenderer/Private/AudioMixerPlatformNonRealtime.cpp
+++ b/Engine/Source/Runtime/NonRealtimeAudioRenderer/Private/AudioMixerPlatformNonRealtime.cpp
@@ -321,14 +321,14 @@ namespace Audio
 #if WITH_XMA2 && USE_XMA2_FOR_STREAMING
 			if (InSoundWave->NumChannels <= 2)
 			{
-				RuntimeFormat = Audio::NAME_XMA;
+				return Audio::NAME_XMA;
 			}
 #endif // WITH_XMA2 && USE_XMA2_FOR_STREAMING
 
 #if USE_VORBIS_FOR_STREAMING
-			RuntimeFormat = Audio::NAME_OGG;
+			return Audio::NAME_OGG;
 #else
-			RuntimeFormat = Audio::NAME_OPUS;
+			return Audio::NAME_OPUS;
 #endif // USE_VORBIS_FOR_STREAMING
 		}
 
