Delete deprecated usage of GetRuntimeFormat and CreateCompressedInfo

diff --git a/Engine/Source/Runtime/Android/AudioMixerAndroid/AudioMixerAndroid.Build.cs b/Engine/Source/Runtime/Android/AudioMixerAndroid/AudioMixerAndroid.Build.cs
--- a/Engine/Source/Runtime/Android/AudioMixerAndroid/AudioMixerAndroid.Build.cs
+++ b/Engine/Source/Runtime/Android/AudioMixerAndroid/AudioMixerAndroid.Build.cs
@@ -14,7 +14,6 @@ public class AudioMixerAndroid : ModuleRules
 			new string[] {
 				"Core",
 				"AudioMixerCore",
-				"BinkAudioDecoder"
 			}
 			);
 			
diff --git a/Engine/Source/Runtime/Android/AudioMixerAndroid/Private/AudioMixerPlatformAndroid.cpp b/Engine/Source/Runtime/Android/AudioMixerAndroid/Private/AudioMixerPlatformAndroid.cpp
--- a/Engine/Source/Runtime/Android/AudioMixerAndroid/Private/AudioMixerPlatformAndroid.cpp
+++ b/Engine/Source/Runtime/Android/AudioMixerAndroid/Private/AudioMixerPlatformAndroid.cpp
@@ -9,8 +9,6 @@
 #include "AudioDevice.h"
 
 #if WITH_ENGINE
-#include "VorbisAudioInfo.h"
-#include "BinkAudioInfo.h"
 #include "AudioPluginUtilities.h"
 #endif
 
@@ -417,41 +415,7 @@ namespace Audio
 			SLresult Result = (*SL_PlayerBufferQueue)->Enqueue(SL_PlayerBufferQueue, Buffer, BufferSize);
 			OPENSLES_LOG_ON_FAIL(Result);
 		}
-	}
-
-	FName FMixerPlatformAndroid::GetRuntimeFormat(const USoundWave* InSoundWave) const
-	{
-		FName RuntimeFormat = Audio::ToName(InSoundWave->GetSoundAssetCompressionType());
-
-		if (RuntimeFormat == Audio::NAME_PLATFORM_SPECIFIC)
-		{
-			RuntimeFormat = Audio::NAME_OGG;
-		}
-
-		return RuntimeFormat;
-	}
-
-	ICompressedAudioInfo* FMixerPlatformAndroid::CreateCompressedAudioInfo(const FName& InRuntimeFormat) const
-	{
-		ICompressedAudioInfo* Decoder = nullptr;
-
-		if (InRuntimeFormat == Audio::NAME_OGG)
-		{
-			Decoder = new FVorbisAudioInfo();
-		}
-		else if (InRuntimeFormat == Audio::NAME_BINKA)
-		{
-#if WITH_BINK_AUDIO
-			Decoder = new FBinkAudioInfo();
-#endif // WITH_BINK_AUDIO			
-		}
-		else
-		{
-			Decoder = Audio::CreateSoundAssetDecoder(InRuntimeFormat);
-		}
-		ensureMsgf(Decoder != nullptr, TEXT("Failed to create a sound asset decoder for compression type: %s"), *InRuntimeFormat.ToString());
-		return Decoder;
-	}
+	}	
 
 	FString FMixerPlatformAndroid::GetDefaultDeviceName()
 	{
diff --git a/Engine/Source/Runtime/Android/AudioMixerAndroid/Private/AudioMixerPlatformAndroid.h b/Engine/Source/Runtime/Android/AudioMixerAndroid/Private/AudioMixerPlatformAndroid.h
--- a/Engine/Source/Runtime/Android/AudioMixerAndroid/Private/AudioMixerPlatformAndroid.h
+++ b/Engine/Source/Runtime/Android/AudioMixerAndroid/Private/AudioMixerPlatformAndroid.h
@@ -33,8 +33,6 @@ namespace Audio
 		virtual bool StopAudioStream() override;
 		virtual FAudioPlatformDeviceInfo GetPlatformDeviceInfo() const override;
 		virtual void SubmitBuffer(const uint8* Buffer) override;
-		virtual FName GetRuntimeFormat(const USoundWave* InSoundWave) const override;
-		virtual ICompressedAudioInfo* CreateCompressedAudioInfo(const FName& InRuntimeFormat) const override;
 		virtual FString GetDefaultDeviceName() override;
 		virtual FAudioPlatformSettings GetPlatformSettings() const override;
 		virtual void SuspendContext() override;
diff --git a/Engine/Source/Runtime/Apple/AudioMixerAudioUnit/AudioMixerAudioUnit.Build.cs b/Engine/Source/Runtime/Apple/AudioMixerAudioUnit/AudioMixerAudioUnit.Build.cs
--- a/Engine/Source/Runtime/Apple/AudioMixerAudioUnit/AudioMixerAudioUnit.Build.cs
+++ b/Engine/Source/Runtime/Apple/AudioMixerAudioUnit/AudioMixerAudioUnit.Build.cs
@@ -27,7 +27,6 @@ public class AudioMixerAudioUnit : ModuleRules
 				new string[]
 				{
 					"Engine",
-					"BinkAudioDecoder"
 				}
 			);
 
diff --git a/Engine/Source/Runtime/Apple/AudioMixerAudioUnit/Private/AudioMixerPlatformAudioUnit.cpp b/Engine/Source/Runtime/Apple/AudioMixerAudioUnit/Private/AudioMixerPlatformAudioUnit.cpp
--- a/Engine/Source/Runtime/Apple/AudioMixerAudioUnit/Private/AudioMixerPlatformAudioUnit.cpp
+++ b/Engine/Source/Runtime/Apple/AudioMixerAudioUnit/Private/AudioMixerPlatformAudioUnit.cpp
@@ -10,8 +10,6 @@
 #include "CoreMinimal.h"
 #include "Misc/ConfigCacheIni.h"
 #include "Misc/CoreDelegates.h"
-#include "ADPCMAudioInfo.h"
-#include "BinkAudioInfo.h"
 #include "Misc/CommandLine.h"
 
 /*
@@ -394,39 +392,6 @@ namespace Audio
 		check(PushResult == BytesToSubmitToAudioMixer);
 	}
 
-	FName FMixerPlatformAudioUnit::GetRuntimeFormat(const USoundWave* InSoundWave) const
-	{
-		FName RuntimeFormat = Audio::ToName(InSoundWave->GetSoundAssetCompressionType());
-
-		// On IOS the default platform specific codec is ADPCM
-		if (RuntimeFormat == Audio::NAME_PLATFORM_SPECIFIC)
-		{
-			RuntimeFormat = Audio::NAME_ADPCM;
-		}
-
-		return RuntimeFormat;
-	}
-
-	ICompressedAudioInfo* FMixerPlatformAudioUnit::CreateCompressedAudioInfo(const FName& InRuntimeFormat) const
-	{
-		ICompressedAudioInfo* Decoder = nullptr;
-
-		if (InRuntimeFormat == Audio::NAME_PLATFORM_SPECIFIC)
-		{
-			Decoder = new FADPCMAudioInfo();
-		}	
-		else if (InRuntimeFormat == Audio::NAME_BINKA)
-		{
-			Decoder = new FBinkAudioInfo();
-		}
-		else
-		{
-			Decoder = Audio::CreateSoundAssetDecoder(InRuntimeFormat);
-		}
-		ensureMsgf(Decoder != nullptr, TEXT("Failed to create a sound asset decoder for compression type: %s"), *InRuntimeFormat.ToString());
-		return Decoder;
-	}
-
 	FString FMixerPlatformAudioUnit::GetDefaultDeviceName()
 	{
 		return FString();
diff --git a/Engine/Source/Runtime/Apple/AudioMixerAudioUnit/Private/AudioMixerPlatformAudioUnit.h b/Engine/Source/Runtime/Apple/AudioMixerAudioUnit/Private/AudioMixerPlatformAudioUnit.h
--- a/Engine/Source/Runtime/Apple/AudioMixerAudioUnit/Private/AudioMixerPlatformAudioUnit.h
+++ b/Engine/Source/Runtime/Apple/AudioMixerAudioUnit/Private/AudioMixerPlatformAudioUnit.h
@@ -35,8 +35,6 @@ namespace Audio
 		virtual bool MoveAudioStreamToNewAudioDevice(const FString& InNewDeviceId) override;
 		virtual FAudioPlatformDeviceInfo GetPlatformDeviceInfo() const override;
 		virtual void SubmitBuffer(const uint8* Buffer) override;
-		virtual FName GetRuntimeFormat(const USoundWave* InSoundWave) const override;
-		virtual ICompressedAudioInfo* CreateCompressedAudioInfo(const FName& InRuntimeFormat) const override;		
 		virtual FString GetDefaultDeviceName() override;
 		virtual FAudioPlatformSettings GetPlatformSettings() const override;
 		virtual int32 GetNumFrames(const int32 InNumReqestedFrames) override;
diff --git a/Engine/Source/Runtime/Apple/AudioMixerCoreAudio/AudioMixerCoreAudio.Build.cs b/Engine/Source/Runtime/Apple/AudioMixerCoreAudio/AudioMixerCoreAudio.Build.cs
--- a/Engine/Source/Runtime/Apple/AudioMixerCoreAudio/AudioMixerCoreAudio.Build.cs
+++ b/Engine/Source/Runtime/Apple/AudioMixerCoreAudio/AudioMixerCoreAudio.Build.cs
@@ -17,7 +17,6 @@ public class AudioMixerCoreAudio : ModuleRules
 				"Core",
 				"CoreUObject",
 				"AudioMixerCore",
-				"BinkAudioDecoder"
 			}
 		);
 
diff --git a/Engine/Source/Runtime/Apple/AudioMixerCoreAudio/Private/AudioMixerPlatformCoreAudio.cpp b/Engine/Source/Runtime/Apple/AudioMixerCoreAudio/Private/AudioMixerPlatformCoreAudio.cpp
--- a/Engine/Source/Runtime/Apple/AudioMixerCoreAudio/Private/AudioMixerPlatformCoreAudio.cpp
+++ b/Engine/Source/Runtime/Apple/AudioMixerCoreAudio/Private/AudioMixerPlatformCoreAudio.cpp
@@ -8,10 +8,7 @@
 #include "CoreMinimal.h"
 #include "Misc/ConfigCacheIni.h"
 #if WITH_ENGINE
-#include "VorbisAudioInfo.h"
-#include "OpusAudioInfo.h"
 #include "AudioDevice.h"
-#include "BinkAudioInfo.h"
 #endif
 /**
 * CoreAudio System Headers
@@ -360,51 +357,6 @@ namespace Audio
 		RemainingBytesInCurrentSubmittedBuffer = BytesPerSubmittedBuffer;
 	}
 
-	FName FMixerPlatformCoreAudio::GetRuntimeFormat(const USoundWave* InSoundWave) const
-	{
-		FName RuntimeFormat = Audio::ToName(InSoundWave->GetSoundAssetCompressionType());
-
-		if (RuntimeFormat == Audio::NAME_PLATFORM_SPECIFIC)
-		{
-			if (InSoundWave->IsStreaming(nullptr))
-			{
-				RuntimeFormat = Audio::NAME_OPUS;
-			}
-			else
-			{
-				RuntimeFormat = Audio::NAME_OGG;
-			}
-		}
-
-		return RuntimeFormat;
-	}
-
-	ICompressedAudioInfo* FMixerPlatformCoreAudio::CreateCompressedAudioInfo(const FName& InRuntimeFormat) const
-	{
-		ICompressedAudioInfo* Decoder = nullptr;
-
-		if (InRuntimeFormat == Audio::NAME_OGG)
-		{
-			Decoder = new FVorbisAudioInfo();
-		}
-		else if (InRuntimeFormat == Audio::NAME_OPUS)
-		{
-			Decoder = new FOpusAudioInfo();
-		}
-#if WITH_BINK_AUDIO
-		else if (InRuntimeFormat == Audio::NAME_BINKA)
-		{
-			return new FBinkAudioInfo();
-		}
-#endif // WITH_BINK_AUDIO	
-		else
-		{
-			Decoder = Audio::CreateSoundAssetDecoder(InRuntimeFormat);
-		}
-		ensureMsgf(Decoder != nullptr, TEXT("Failed to create a sound asset decoder for compression type: %s"), *InRuntimeFormat.ToString());
-		return Decoder;
-	}
-
 	FString FMixerPlatformCoreAudio::GetDefaultDeviceName()
 	{
 		return FString();
diff --git a/Engine/Source/Runtime/Apple/AudioMixerCoreAudio/Private/AudioMixerPlatformCoreAudio.h b/Engine/Source/Runtime/Apple/AudioMixerCoreAudio/Private/AudioMixerPlatformCoreAudio.h
--- a/Engine/Source/Runtime/Apple/AudioMixerCoreAudio/Private/AudioMixerPlatformCoreAudio.h
+++ b/Engine/Source/Runtime/Apple/AudioMixerCoreAudio/Private/AudioMixerPlatformCoreAudio.h
@@ -35,8 +35,6 @@ namespace Audio
 		virtual bool MoveAudioStreamToNewAudioDevice(const FString& InNewDeviceId) override;
 		virtual FAudioPlatformDeviceInfo GetPlatformDeviceInfo() const override;
 		virtual void SubmitBuffer(const uint8* Buffer) override;
-		virtual FName GetRuntimeFormat(const USoundWave* InSoundWave) const override;
-		virtual ICompressedAudioInfo* CreateCompressedAudioInfo(const FName& InRuntimeFormat) const override;
 		virtual FString GetDefaultDeviceName() override;
 		virtual FAudioPlatformSettings GetPlatformSettings() const override;
         virtual int32 GetNumFrames(const int32 InNumReqestedFrames) override;
diff --git a/Engine/Source/Runtime/AudioLink/AudioMixerPlatformAudioLink/AudioMixerPlatformAudioLink.Build.cs b/Engine/Source/Runtime/AudioLink/AudioMixerPlatformAudioLink/AudioMixerPlatformAudioLink.Build.cs
--- a/Engine/Source/Runtime/AudioLink/AudioMixerPlatformAudioLink/AudioMixerPlatformAudioLink.Build.cs
+++ b/Engine/Source/Runtime/AudioLink/AudioMixerPlatformAudioLink/AudioMixerPlatformAudioLink.Build.cs
@@ -19,15 +19,5 @@ public class AudioMixerPlatformAudioLink: ModuleRules
 				"Engine",
 		});
 					
-		if (Target.bCompileAgainstEngine) 
-		{
-			PrivateDependencyModuleNames.Add("BinkAudioDecoder");
-
-			AddEngineThirdPartyPrivateStaticDependencies(Target,
-				"UEOgg",
-				"Vorbis",
-				"VorbisFile"
-				);
-		}
 	}
 }
diff --git a/Engine/Source/Runtime/AudioMixer/Private/AudioMixerDevice.cpp b/Engine/Source/Runtime/AudioMixer/Private/AudioMixerDevice.cpp
--- a/Engine/Source/Runtime/AudioMixer/Private/AudioMixerDevice.cpp
+++ b/Engine/Source/Runtime/AudioMixer/Private/AudioMixerDevice.cpp
@@ -825,11 +825,6 @@ namespace Audio
 		return new FMixerSource(this);
 	}
 
-	FName FMixerDevice::GetRuntimeFormat(const USoundWave* InSoundWave) const
-	{
-		return InSoundWave->GetRuntimeFormat();
-    }
-
 	bool FMixerDevice::HasCompressedAudioInfoClass(USoundWave* InSoundWave)
 	{
 		check(InSoundWave);
@@ -849,25 +844,6 @@ namespace Audio
 		return AudioMixerPlatform->DisablePCMAudioCaching();
 	}
 	
-	ICompressedAudioInfo* FMixerDevice::CreateAudioInfo(FName InFormat) const
-	{
-		if (IAudioInfoFactory* Factory = IAudioInfoFactoryRegistry::Get().Find(InFormat))
-		{
-			return Factory->Create();
-		}
-		return nullptr;
-	}
-
-	ICompressedAudioInfo* FMixerDevice::CreateCompressedAudioInfo(const USoundWave* InSoundWave) const
-	{
-		return CreateAudioInfo(GetRuntimeFormat(InSoundWave));
-	}
-
-	class ICompressedAudioInfo* FMixerDevice::CreateCompressedAudioInfo(const FSoundWaveProxyPtr& InSoundWaveProxy) const
-	{
-		return CreateAudioInfo(InSoundWaveProxy->GetRuntimeFormat());
-	}
-
 	bool FMixerDevice::ValidateAPICall(const TCHAR* Function, uint32 ErrorCode)
 	{
 		return false;
diff --git a/Engine/Source/Runtime/AudioMixer/Public/AudioMixerDevice.h b/Engine/Source/Runtime/AudioMixer/Public/AudioMixerDevice.h
--- a/Engine/Source/Runtime/AudioMixer/Public/AudioMixerDevice.h
+++ b/Engine/Source/Runtime/AudioMixer/Public/AudioMixerDevice.h
@@ -132,12 +132,9 @@ namespace Audio
 		AUDIOMIXER_API virtual double GetInterpolatedAudioClock() const override;
 		AUDIOMIXER_API virtual FAudioEffectsManager* CreateEffectsManager() override;
 		AUDIOMIXER_API virtual FSoundSource* CreateSoundSource() override;
-		AUDIOMIXER_API virtual FName GetRuntimeFormat(const USoundWave* SoundWave) const override;
 		AUDIOMIXER_API virtual bool HasCompressedAudioInfoClass(USoundWave* SoundWave) override;
 		AUDIOMIXER_API virtual bool SupportsRealtimeDecompression() const override;
 		AUDIOMIXER_API virtual bool DisablePCMAudioCaching() const override;
-		AUDIOMIXER_API virtual class ICompressedAudioInfo* CreateCompressedAudioInfo(const USoundWave* SoundWave) const override;
-		AUDIOMIXER_API virtual class ICompressedAudioInfo* CreateCompressedAudioInfo(const FSoundWaveProxyPtr& SoundWave) const override;
 		AUDIOMIXER_API virtual bool ValidateAPICall(const TCHAR* Function, uint32 ErrorCode) override;
 #if UE_ALLOW_EXEC_COMMANDS
 		AUDIOMIXER_API virtual bool Exec(UWorld* InWorld, const TCHAR* Cmd, FOutputDevice& Ar) override;
@@ -413,8 +410,6 @@ namespace Audio
 
 		void UnloadSoundSubmix(const USoundSubmixBase& SoundSubmix, const bool bReparentChildren);
 
-		ICompressedAudioInfo* CreateAudioInfo(FName InFormat) const;
-
 		bool IsRequiredSubmixType(const USoundSubmixBase* InSubmix) const;
 		FMixerSubmixPtr GetRequiredSubmixInstance(uint32 InSubmixId) const;
 		FMixerSubmixPtr GetRequiredSubmixInstance(const USoundSubmixBase* InSubmix) const;
diff --git a/Engine/Source/Runtime/AudioMixerCore/Public/AudioMixer.h b/Engine/Source/Runtime/AudioMixerCore/Public/AudioMixer.h
--- a/Engine/Source/Runtime/AudioMixerCore/Public/AudioMixer.h
+++ b/Engine/Source/Runtime/AudioMixerCore/Public/AudioMixer.h
@@ -520,14 +520,6 @@ namespace Audio
 		/** Whether this is an interface for a non-realtime renderer. If true, synch events will behave differently to avoid deadlocks. */
 		virtual bool IsNonRealtime() const { return false; }
 
-		/** Returns the FName version of a given sound wave's format. Override to provide a runtime format (codec) that is platform specific. */
-		UE_DEPRECATED(5.2, "GetRuntimeFormat is now deprecated. Please use USoundWave::GetRuntimeFormat()")
-		virtual FName GetRuntimeFormat(const USoundWave* InSoundWave) const { return FName(); }
-
-		/** Creates a Compressed audio info class suitable for decompressing this SoundWave. OVerride to create platform-specific decoders. */
-		UE_DEPRECATED(5.2, "CreateCompressedAudioInfo is now deprecated. Please use IAudioInfoFactory::Create()")
-		virtual ICompressedAudioInfo* CreateCompressedAudioInfo(const FName& InRuntimeFormat) const { return nullptr; };
-
 		/** Return any optional device name defined in platform configuratio. */
 		virtual FString GetDefaultDeviceName() = 0;
 
diff --git a/Engine/Source/Runtime/Engine/Engine.Build.cs b/Engine/Source/Runtime/Engine/Engine.Build.cs
--- a/Engine/Source/Runtime/Engine/Engine.Build.cs
+++ b/Engine/Source/Runtime/Engine/Engine.Build.cs
@@ -176,7 +176,8 @@ public class Engine : ModuleRules
 			"UEOgg",
 			"Vorbis",
 			"VorbisFile",
-			"libOpus"
+			"libOpus",
+			"BinkAudioDecoder"
 			);
 
 		DynamicallyLoadedModuleNames.Add("EyeTracker");
diff --git a/Engine/Source/Runtime/Engine/Private/AudioDeviceManager.cpp b/Engine/Source/Runtime/Engine/Private/AudioDeviceManager.cpp
--- a/Engine/Source/Runtime/Engine/Private/AudioDeviceManager.cpp
+++ b/Engine/Source/Runtime/Engine/Private/AudioDeviceManager.cpp
@@ -356,6 +356,9 @@ void FAudioDeviceManager::RegisterAudioInfoFactories()
 			FModuleManager::Get().LoadModuleChecked(*i);
 		}
 	}
+	
+	// Vorbis needs a load call.
+	LoadVorbisLibraries();
 
 	// Register the engine formats.
 	EngineFormats.Add(MakePimpl<FSimpleAudioInfoFactory>([] { return new FADPCMAudioInfo(); }, Audio::NAME_PCM));
diff --git a/Engine/Source/Runtime/Engine/Public/AudioDevice.h b/Engine/Source/Runtime/Engine/Public/AudioDevice.h
--- a/Engine/Source/Runtime/Engine/Public/AudioDevice.h
+++ b/Engine/Source/Runtime/Engine/Public/AudioDevice.h
@@ -1166,9 +1166,6 @@ public:
 	 */
 	ENGINE_API void DeactivateReverbEffect(FName TagName);
 
-	UE_DEPRECATED(5.2, "GetRuntimeFormat is now deprecated. Please call SoundWave::GetRuntimeFormat() instead")
-	virtual FName GetRuntimeFormat(const USoundWave* SoundWave) const = 0;
-
 	/** Whether this SoundWave has an associated info class to decompress it */
 	virtual bool HasCompressedAudioInfoClass(USoundWave* SoundWave) { return false; }
 
@@ -1184,13 +1181,6 @@ public:
 		return false;
 	}
 
-	/** Creates a Compressed audio info class suitable for decompressing this SoundWave */
-	UE_DEPRECATED(5.2, "CreateCompressedAudioInfo is now deprecated. Please use IAudioInfoFactory::Create() instead")
-	virtual ICompressedAudioInfo* CreateCompressedAudioInfo(const USoundWave* SoundWave) const { return nullptr; }
-	
-	UE_DEPRECATED(5.2, "CreateCompressedAudioInfo is now deprecated. Please use IAudioInfoFactory::Create() instead")
-	virtual ICompressedAudioInfo* CreateCompressedAudioInfo(const FSoundWaveProxyPtr& SoundWave) const { return nullptr; }
-
 	/**
 	 * Check for errors and output a human readable string
 	 */
diff --git a/Engine/Source/Runtime/IOS/IOSAudio/Public/IOSAudioDevice.h b/Engine/Source/Runtime/IOS/IOSAudio/Public/IOSAudioDevice.h
--- a/Engine/Source/Runtime/IOS/IOSAudio/Public/IOSAudioDevice.h
+++ b/Engine/Source/Runtime/IOS/IOSAudio/Public/IOSAudioDevice.h
@@ -193,13 +193,7 @@ class FIOSAudioDevice : public FAudioDevice
 {
 public:
 	FIOSAudioDevice();
-	virtual ~FIOSAudioDevice() { }
-	
-	virtual FName GetRuntimeFormat(const USoundWave* SoundWave) const override
-	{
-		static FName NAME_ADPCM(TEXT("ADPCM"));
-		return NAME_ADPCM;
-	}
+	virtual ~FIOSAudioDevice() { }		
 
 	AUGraph GetAudioUnitGraph() const { return AudioUnitGraph; }
 	AUNode GetMixerNode() const { return MixerNode; }
diff --git a/Engine/Source/Runtime/Linux/AudioMixerSDL/AudioMixerSDL.Build.cs b/Engine/Source/Runtime/Linux/AudioMixerSDL/AudioMixerSDL.Build.cs
--- a/Engine/Source/Runtime/Linux/AudioMixerSDL/AudioMixerSDL.Build.cs
+++ b/Engine/Source/Runtime/Linux/AudioMixerSDL/AudioMixerSDL.Build.cs
@@ -20,13 +20,6 @@ public class AudioMixerSDL : ModuleRules
 			PrivateIncludePaths.Add("Runtime/Linux/AudioMixerSDL/Private/" + PlatformName);
 		}
 
-		{
-			// Bink Audio isn't yet built for other SDL platforms, however
-			// the module won't provide a lib and will define WITH_BINK_AUDIO
-			// to 0 for those platforms.
-			PrivateDependencyModuleNames.Add("BinkAudioDecoder");
-		}
-
 		PrivateDependencyModuleNames.AddRange(
 			new string[] {
 				"Core",
diff --git a/Engine/Source/Runtime/Linux/AudioMixerSDL/Private/AudioMixerPlatformSDL.cpp b/Engine/Source/Runtime/Linux/AudioMixerSDL/Private/AudioMixerPlatformSDL.cpp
--- a/Engine/Source/Runtime/Linux/AudioMixerSDL/Private/AudioMixerPlatformSDL.cpp
+++ b/Engine/Source/Runtime/Linux/AudioMixerSDL/Private/AudioMixerPlatformSDL.cpp
@@ -10,13 +10,6 @@
 
 #if WITH_ENGINE
 #include "AudioPluginUtilities.h"
-#include "OpusAudioInfo.h"
-#include "VorbisAudioInfo.h"
-#include "ADPCMAudioInfo.h"
-#if WITH_BINK_AUDIO
-#include "BinkAudioInfo.h"
-#endif // WITH_BINK_AUDIO
-
 #endif // WITH_ENGINE
 
 DECLARE_LOG_CATEGORY_EXTERN(LogAudioMixerSDL, Log, All);
@@ -374,59 +367,6 @@ namespace Audio
 		ReadNextBuffer();
 	}
 
-	FName FMixerPlatformSDL::GetRuntimeFormat(const USoundWave* InSoundWave) const
-	{
-		FName RuntimeFormat = Audio::ToName(InSoundWave->GetSoundAssetCompressionType());
-
-		if (RuntimeFormat == Audio::NAME_PLATFORM_SPECIFIC)
-		{
-			if (InSoundWave->IsStreaming(nullptr))
-			{
-				return Audio::NAME_OPUS;
-			}
-			else
-			{
-				return Audio::NAME_OGG;
-			}
-		}
-
-#if !WITH_BINK_AUDIO
-		// Some SDL platforms don't yet support bink, so fall back to ADPCM
-		if (RuntimeFormat == Audio::NAME_BINKA)
-		{
-			return Audio::NAME_ADPCM;
-		}
-#endif
-
-		return RuntimeFormat;
-	}
-
-	ICompressedAudioInfo* FMixerPlatformSDL::CreateCompressedAudioInfo(const FName& InRuntimeFormat) const
-	{
-		ICompressedAudioInfo* Decoder = nullptr;
-
-		if (InRuntimeFormat == Audio::NAME_OGG)
-		{
-			Decoder = new FVorbisAudioInfo();
-		}
-		else if (InRuntimeFormat == Audio::NAME_OPUS)
-		{
-			Decoder = new FOpusAudioInfo();
-		}
-#if WITH_BINK_AUDIO
-		else if (InRuntimeFormat == Audio::NAME_BINKA)
-		{
-			Decoder = new FBinkAudioInfo();
-		}
-#endif // WITH_BINK_AUDIO	
-		else
-		{
-			Decoder = Audio::CreateSoundAssetDecoder(InRuntimeFormat);
-		}
-		ensureMsgf(Decoder != nullptr, TEXT("Failed to create a sound asset decoder for compression type: %s"), *InRuntimeFormat.ToString());
-		return Decoder;
-	}
-
 	FString FMixerPlatformSDL::GetDefaultDeviceName()
 	{
 		static FString DefaultName(TEXT("Default SDL Audio Device."));
diff --git a/Engine/Source/Runtime/Linux/AudioMixerSDL/Private/AudioMixerPlatformSDL.h b/Engine/Source/Runtime/Linux/AudioMixerSDL/Private/AudioMixerPlatformSDL.h
--- a/Engine/Source/Runtime/Linux/AudioMixerSDL/Private/AudioMixerPlatformSDL.h
+++ b/Engine/Source/Runtime/Linux/AudioMixerSDL/Private/AudioMixerPlatformSDL.h
@@ -34,8 +34,6 @@ namespace Audio
 		bool StopAudioStream() override;
 		FAudioPlatformDeviceInfo GetPlatformDeviceInfo() const override;
 		void SubmitBuffer(const uint8* Buffer) override;
-		FName GetRuntimeFormat(const USoundWave* InSoundWave) const override;
-		ICompressedAudioInfo* CreateCompressedAudioInfo(const FName& InRuntimeFormat) const override;
 		FString GetDefaultDeviceName() override;
 		FAudioPlatformSettings GetPlatformSettings() const override;
 		void ResumeContext() override;
diff --git a/Engine/Source/Runtime/NonRealtimeAudioRenderer/NonRealtimeAudioRenderer.Build.cs b/Engine/Source/Runtime/NonRealtimeAudioRenderer/NonRealtimeAudioRenderer.Build.cs
--- a/Engine/Source/Runtime/NonRealtimeAudioRenderer/NonRealtimeAudioRenderer.Build.cs
+++ b/Engine/Source/Runtime/NonRealtimeAudioRenderer/NonRealtimeAudioRenderer.Build.cs
@@ -15,7 +15,6 @@ namespace UnrealBuildTool.Rules
                     "Core",
                     "CoreUObject",
                     "Engine",
-                    "BinkAudioDecoder",
                     "AudioMixerCore",
                     "SignalProcessing",
                     "AudioMixer",
diff --git a/Engine/Source/Runtime/NonRealtimeAudioRenderer/Private/AudioMixerPlatformNonRealtime.cpp b/Engine/Source/Runtime/NonRealtimeAudioRenderer/Private/AudioMixerPlatformNonRealtime.cpp
--- a/Engine/Source/Runtime/NonRealtimeAudioRenderer/Private/AudioMixerPlatformNonRealtime.cpp
+++ b/Engine/Source/Runtime/NonRealtimeAudioRenderer/Private/AudioMixerPlatformNonRealtime.cpp
@@ -7,23 +7,10 @@
 #include "HAL/PlatformAffinity.h"
 #include "Misc/App.h"
 
-#ifndef WITH_XMA2
-#define WITH_XMA2 0
-#endif
-
 #ifndef HAS_COMPRESSED_AUDIO_INFO_CLASS
 #define HAS_COMPRESSED_AUDIO_INFO_CLASS 0
 #endif
 
-#if WITH_XMA2
-#include "XMAAudioInfo.h"
-#endif  //#if WITH_XMA2
-#if WITH_BINK_AUDIO
-#include "BinkAudioInfo.h"
-#endif // WITH_BINK_AUDIO
-#include "OpusAudioInfo.h"
-#include "VorbisAudioInfo.h"
-#include "ADPCMAudioInfo.h"
 #include "Interfaces/IAudioFormat.h"
 #include "AudioDevice.h"
 
@@ -127,14 +114,6 @@ namespace Audio
 			return false;
 		}
 
-#if WITH_XMA2
-		//Initialize our XMA2 decoder context
-		XMA2_INFO_CALL(FXMAAudioInfo::Initialize());
-#endif //#if WITH_XMA2
-
-		// Load ogg and vorbis dlls if they haven't been loaded yet
-		LoadVorbisLibraries();
-
 		bIsInitialized = true;
 
 		TickDelta = FApp::GetDeltaTime();
@@ -149,10 +128,6 @@ namespace Audio
 			return false;
 		}
 
-#if WITH_XMA2
-		XMA2_INFO_CALL(FXMAAudioInfo::Shutdown());
-#endif
-
 		bIsInitialized = false;
 
 		return true;
@@ -308,63 +283,6 @@ namespace Audio
 		}
 	}
 
-#if WITH_XMA2
-	static FName NAME_XMA(TEXT("XMA"));
-#endif
-
-	FName FMixerPlatformNonRealtime::GetRuntimeFormat(const USoundWave* InSoundWave) const
-	{
-		FName RuntimeFormat = Audio::ToName(InSoundWave->GetSoundAssetCompressionType());
-
-		if (RuntimeFormat == Audio::NAME_PLATFORM_SPECIFIC)
-		{
-#if WITH_XMA2 && USE_XMA2_FOR_STREAMING
-			if (InSoundWave->NumChannels <= 2)
-			{
-				return Audio::NAME_XMA;
-			}
-#endif // WITH_XMA2 && USE_XMA2_FOR_STREAMING
-
-#if USE_VORBIS_FOR_STREAMING
-			return Audio::NAME_OGG;
-#else
-			return Audio::NAME_OPUS;
-#endif // USE_VORBIS_FOR_STREAMING
-		}
-
-		return RuntimeFormat;
-	}
-
-	ICompressedAudioInfo* FMixerPlatformNonRealtime::CreateCompressedAudioInfo(const FName& InRuntimeFormat) const
-	{
-		// Need to create a platform-specific codec
-#if WITH_XMA2 && USE_XMA2_FOR_STREAMING
-		if (InRuntimeFormat == NAME_XMA)
-		{
-			return XMA2_INFO_NEW();
-		}
-#endif // WITH_XMA2 && USE_XMA2_FOR_STREAMING			
-#if USE_VORBIS_FOR_STREAMING
-		if (InRuntimeFormat == Audio::NAME_OGG)
-		{
-			return new FVorbisAudioInfo();
-		}
-#else // USE_VORBIS_FOR_STREAMING
-		if (InRuntimeFormat == Audio::NAME_OPUS)
-		{
-			return new FOpusAudioInfo();
-		}
-#endif // #if USE_VORBIS_FOR_STREAMING
-#if WITH_BINK_AUDIO
-		if (InRuntimeFormat == Audio::NAME_BINKA)
-		{
-			return new FBinkAudioInfo();
-		}
-#endif
-
-		return Audio::CreateSoundAssetDecoder(InRuntimeFormat);
-	}
-
 	FString FMixerPlatformNonRealtime::GetDefaultDeviceName()
 	{
 		//GConfig->GetString(TEXT("/Script/WindowsTargetPlatform.WindowsTargetSettings"), TEXT("AudioDevice"), WindowsAudioDeviceName, GEngineIni);
@@ -378,10 +296,6 @@ namespace Audio
 
 	void FMixerPlatformNonRealtime::OnHardwareUpdate()
 	{
-#if WITH_XMA2
-		XMA2_INFO_CALL(FXMAAudioInfo::Tick());
-#endif //WITH_XMA2
-
 		if (RenderEveryTickCvar)
 		{
 			RenderAudio(TickDelta);
diff --git a/Engine/Source/Runtime/NonRealtimeAudioRenderer/Public/AudioMixerPlatformNonRealtime.h b/Engine/Source/Runtime/NonRealtimeAudioRenderer/Public/AudioMixerPlatformNonRealtime.h
--- a/Engine/Source/Runtime/NonRealtimeAudioRenderer/Public/AudioMixerPlatformNonRealtime.h
+++ b/Engine/Source/Runtime/NonRealtimeAudioRenderer/Public/AudioMixerPlatformNonRealtime.h
@@ -37,9 +37,7 @@ namespace Audio
 		virtual void ResumePlaybackOnNewDevice() override;
 		virtual FAudioPlatformDeviceInfo GetPlatformDeviceInfo() const override;
 		virtual void SubmitBuffer(const uint8* Buffer) override;
-		virtual bool DisablePCMAudioCaching() const override;
-		virtual FName GetRuntimeFormat(const USoundWave* InSoundWave) const override;
-		virtual ICompressedAudioInfo* CreateCompressedAudioInfo(const FName& InRuntimeFormat) const override;
+		virtual bool DisablePCMAudioCaching() const override;		
 		virtual FString GetDefaultDeviceName() override;
 		virtual FAudioPlatformSettings GetPlatformSettings() const override;
 		virtual void OnHardwareUpdate() override;
diff --git a/Engine/Source/Runtime/Windows/AudioMixerXAudio2/AudioMixerXAudio2.Build.cs b/Engine/Source/Runtime/Windows/AudioMixerXAudio2/AudioMixerXAudio2.Build.cs
--- a/Engine/Source/Runtime/Windows/AudioMixerXAudio2/AudioMixerXAudio2.Build.cs
+++ b/Engine/Source/Runtime/Windows/AudioMixerXAudio2/AudioMixerXAudio2.Build.cs
@@ -10,12 +10,10 @@ public class AudioMixerXAudio2 : ModuleRules
 
 		if (Target.bCompileAgainstEngine)
         {
-			// Engine module is required for CompressedAudioInfo implementations.
             PrivateDependencyModuleNames.AddRange(
 				new string[]
 				{
 					"Engine",
-					"BinkAudioDecoder"
 				}
 			);
 
diff --git a/Engine/Source/Runtime/Windows/AudioMixerXAudio2/Private/AudioMixerPlatformXAudio2.cpp b/Engine/Source/Runtime/Windows/AudioMixerXAudio2/Private/AudioMixerPlatformXAudio2.cpp
--- a/Engine/Source/Runtime/Windows/AudioMixerXAudio2/Private/AudioMixerPlatformXAudio2.cpp
+++ b/Engine/Source/Runtime/Windows/AudioMixerXAudio2/Private/AudioMixerPlatformXAudio2.cpp
@@ -17,19 +17,11 @@
 #ifndef WITH_XMA2
 #define WITH_XMA2 0
 #endif
-
 #if WITH_ENGINE
 #if WITH_XMA2
 #include "XMAAudioInfo.h"
 #endif  //#if WITH_XMA2
-#include "OpusAudioInfo.h"
-#include "VorbisAudioInfo.h"
-#include "AudioPluginUtilities.h"
-#if WITH_BINK_AUDIO
-#include "BinkAudioInfo.h"
-#endif // WITH_BINK_AUDIO
 #endif //WITH_ENGINE
-
 #include "CoreGlobals.h"
 #include "Misc/ConfigCacheIni.h"
 #include "Misc/MessageDialog.h"
@@ -674,10 +666,6 @@ namespace Audio
 		//Initialize our XMA2 decoder context
 		XMA2_INFO_CALL(FXMAAudioInfo::Initialize());
 #endif //#if WITH_XMA2
-#if WITH_ENGINE
-		// Load ogg and vorbis dlls if they haven't been loaded yet
-		LoadVorbisLibraries();
-#endif // WITH_ENGINE
 
 		if(IAudioMixer::ShouldRecycleThreads())
 		{
@@ -1847,72 +1835,6 @@ namespace Audio
 		}
 	}
 
-#if WITH_XMA2
-	static FName NAME_XMA(TEXT("XMA"));
-#endif
-
-	FName FMixerPlatformXAudio2::GetRuntimeFormat(const USoundWave* InSoundWave) const
-	{
-		FName RuntimeFormat = Audio::ToName(InSoundWave->GetSoundAssetCompressionType());
-
-		if (RuntimeFormat == Audio::NAME_PLATFORM_SPECIFIC)
-		{
-#if WITH_XMA2 && USE_XMA2_FOR_STREAMING
-			if (InSoundWave->NumChannels <= 2)
-			{
-				return Audio::NAME_XMA;
-			}
-#endif // WITH_XMA2 && USE_XMA2_FOR_STREAMING
-
-#if USE_VORBIS_FOR_STREAMING
-			return Audio::NAME_OGG;
-#else
-			return Audio::NAME_OPUS;
-#endif // USE_VORBIS_FOR_STREAMING
-		}
-
-		return RuntimeFormat;
-	}
-
-	ICompressedAudioInfo* FMixerPlatformXAudio2::CreateCompressedAudioInfo(const FName& InRuntimeFormat) const
-	{
-		ICompressedAudioInfo* Decoder = nullptr;
-
-		// Need to create a platform-specific codec
-#if WITH_XMA2 && USE_XMA2_FOR_STREAMING
-		if (InRuntimeFormat == Audio::NAME_XMA)
-		{
-			Decoder = XMA2_INFO_NEW();
-		}
-		else
-#endif // WITH_XMA2 && USE_XMA2_FOR_STREAMING			
-		if (InRuntimeFormat == Audio::NAME_OGG)
-		{
-#if USE_VORBIS_FOR_STREAMING
-			Decoder = new FVorbisAudioInfo();
-#endif // #if USE_VORBIS_FOR_STREAMING
-		}
-		else if (InRuntimeFormat == Audio::NAME_OPUS)
-		{
-#if !USE_VORBIS_FOR_STREAMING
-			Decoder = new FOpusAudioInfo();
-#endif // #if !USE_VORBIS_FOR_STREAMING
-		}
-#if WITH_BINK_AUDIO
-		else if (InRuntimeFormat == Audio::NAME_BINKA)
-		{
-			Decoder = new FBinkAudioInfo();
-		}
-#endif // #if WITH_BINK_AUDIO
-		else
-		{
-			// Fallback to the multiplatform decoder
-			Decoder = Audio::CreateSoundAssetDecoder(InRuntimeFormat);
-		}
-		ensureMsgf(Decoder != nullptr, TEXT("Failed to create a sound asset decoder for compression type: %s"), *InRuntimeFormat.ToString());
-		return Decoder;
-	}
-
 	FString FMixerPlatformXAudio2::GetDefaultDeviceName()
 	{
 		return FString();
diff --git a/Engine/Source/Runtime/Windows/AudioMixerXAudio2/Private/AudioMixerPlatformXAudio2.h b/Engine/Source/Runtime/Windows/AudioMixerXAudio2/Private/AudioMixerPlatformXAudio2.h
--- a/Engine/Source/Runtime/Windows/AudioMixerXAudio2/Private/AudioMixerPlatformXAudio2.h
+++ b/Engine/Source/Runtime/Windows/AudioMixerXAudio2/Private/AudioMixerPlatformXAudio2.h
@@ -82,8 +82,6 @@ namespace Audio
 		virtual FAudioPlatformDeviceInfo GetPlatformDeviceInfo() const override;
 		virtual void SubmitBuffer(const uint8* Buffer) override;
 		virtual bool DisablePCMAudioCaching() const override;
-		virtual FName GetRuntimeFormat(const USoundWave* InSoundWave) const override;
-		virtual ICompressedAudioInfo* CreateCompressedAudioInfo(const FName& InRuntimeFormat) const override;
 		virtual FString GetDefaultDeviceName() override;
 		virtual FAudioPlatformSettings GetPlatformSettings() const override;
 		virtual void OnHardwareUpdate() override;
